Manually resolved conflicts from 4.17 merge.

diff --git a/Engine/Plugins/Experimental/ImagePlate/Source/ImagePlate/Private/ImagePlateFileSequence.cpp b/Engine/Plugins/Experimental/ImagePlate/Source/ImagePlate/Private/ImagePlateFileSequence.cpp
--- a/Engine/Plugins/Experimental/ImagePlate/Source/ImagePlate/Private/ImagePlateFileSequence.cpp
+++ b/Engine/Plugins/Experimental/ImagePlate/Source/ImagePlate/Private/ImagePlateFileSequence.cpp
@@ -14,6 +14,7 @@
 #include "Async.h"
 #include "PackageName.h"
 #include "Engine/Texture2DDynamic.h"
+#include "Engine/TextureRenderTarget2D.h"
 
 DEFINE_LOG_CATEGORY_STATIC(LogImagePlateFileSequence, Log, Warning);
 
@@ -31,7 +32,12 @@ UImagePlateFileSequence::UImagePlateFileSequence(const FObjectInitializer& Init)
 
 FImagePlateAsyncCache UImagePlateFileSequence::GetAsyncCache()
 {
-	FString Path = SequencePath.Path;
+	FString Path;
+	if (!FPackageName::TryConvertLongPackageNameToFilename(SequencePath.Path, Path))
+	{
+		UE_LOG(LogImagePlateFileSequence, Warning, TEXT("Sequence path is not a long package name. This path is not portable, and may not work in a packaged build."));
+		Path = SequencePath.Path;
+	}
 
 	if (GetDefault<UImagePlateSettings>()->ProxyName.Len() != 0)
 	{
@@ -65,50 +71,92 @@ bool FImagePlateSourceFrame::IsValid() const
 	return Buffer.IsValid();
 }
 
-TFuture<void> FImagePlateSourceFrame::CopyTo(UTexture2DDynamic* DestinationTexture)
+bool FImagePlateSourceFrame::EnsureTextureMetrics(UTexture* DestinationTexture) const
 {
-	// Copy the data to the texture via a render command
-	struct FRenderCommandData
+	if (BitDepth != 16 && BitDepth != 8)
 	{
-		TPromise<void> Promise;
-		bool bClearTexture;
-		FImagePlateSourceFrame SourceFrame;
-		UTexture2DDynamic* DestinationTexture;
-	};
+		UE_LOG(LogImagePlateFileSequence, Warning, TEXT("Unsupported source image bitdepth: %u. Only float 16 and fixed 8 bitdepths are supported"), BitDepth);
+		return false;
+	}
 
-	// Shared render command data to allow us to pass the promise through to the command
-	typedef TSharedPtr<FRenderCommandData, ESPMode::ThreadSafe> FCommandDataPtr;
-	FCommandDataPtr CommandData = MakeShared<FRenderCommandData, ESPMode::ThreadSafe>();
+	// Ensure the texture dimensions and bitdepth match if possible
+	bool bNeedsUpdate = false;
 
-	TFuture<void> CompletionFuture = CommandData->Promise.GetFuture();
+	if (auto* Texture2DDynamic = Cast<UTexture2DDynamic>(DestinationTexture))
+	{
+		if (Width > 0 && Texture2DDynamic->SizeX != Width)
+		{
+			Texture2DDynamic->SizeX = Width;
+			bNeedsUpdate = true;
+		}
 
-	CommandData->bClearTexture = !Buffer.IsValid();
-	CommandData->SourceFrame = *this;
-	CommandData->DestinationTexture = DestinationTexture;
+		if (Height > 0 && Texture2DDynamic->SizeY != Height)
+		{
+			Texture2DDynamic->SizeY = Height;
+			bNeedsUpdate = true;
+		}
 
-	// Ensure the texture dimensions and bitdepth match
-	bool bNeedsUpdate = false;
-	if (Width > 0 && DestinationTexture->SizeX != Width)
-	{
-		DestinationTexture->SizeX = Width;
-		bNeedsUpdate = true;
+		uint32 DestBitDepth = GPixelFormats[Texture2DDynamic->Format].BlockBytes * 8 / GPixelFormats[Texture2DDynamic->Format].NumComponents;
+		if (DestBitDepth != BitDepth)
+		{
+			bNeedsUpdate = true;
+			switch (BitDepth)
+			{
+			case 16: Texture2DDynamic->Format = PF_FloatRGBA; break;
+			case 8: Texture2DDynamic->Format = PF_R8G8B8A8; break;
+			}
+		}
 	}
-
-	if (Height > 0 && DestinationTexture->SizeY != Height)
+	else if (auto* TextureRenderTarget2D = Cast<UTextureRenderTarget2D>(DestinationTexture))
 	{
-		DestinationTexture->SizeY = Height;
-		bNeedsUpdate = true;
-	}
+		if (Width > 0 && TextureRenderTarget2D->SizeX != Width)
+		{
+			TextureRenderTarget2D->SizeX = Width;
+			bNeedsUpdate = true;
+		}
 
-	uint32 DestBitDepth = GPixelFormats[DestinationTexture->Format].BlockBytes * 8 / GPixelFormats[DestinationTexture->Format].NumComponents;
-	if (DestBitDepth != BitDepth)
+		if (Height > 0 && TextureRenderTarget2D->SizeY != Height)
+		{
+			TextureRenderTarget2D->SizeY = Height;
+			bNeedsUpdate = true;
+		}
+
+		EPixelFormat PF = GetPixelFormatFromRenderTargetFormat(TextureRenderTarget2D->RenderTargetFormat);
+		uint32 DestBitDepth = GPixelFormats[PF].BlockBytes * 8 / GPixelFormats[PF].NumComponents;
+		if (DestBitDepth != BitDepth)
+		{
+			bNeedsUpdate = true;
+			switch (BitDepth)
+			{
+			case 16: TextureRenderTarget2D->RenderTargetFormat = RTF_RGBA16f; break;
+			case 8: TextureRenderTarget2D->RenderTargetFormat = RTF_RGBA8; break;
+			}
+		}
+	}
+	else if (DestinationTexture->Resource)
 	{
-		bNeedsUpdate = true;
-		switch (BitDepth)
+		// We have to have a valid texture to check whether it's the right type or not
+		if (!DestinationTexture->Resource->TextureRHI)
 		{
-		case 16: DestinationTexture->Format = PF_FloatRGBA; break;
-		case 32: DestinationTexture->Format = PF_R8G8B8A8; break;
-		default: CommandData->bClearTexture = true;
+			DestinationTexture->UpdateResource();
+			FlushRenderingCommands();
+		}
+
+		FTexture2DRHIRef Texture2DRHI = DestinationTexture->Resource->TextureRHI ? DestinationTexture->Resource->TextureRHI->GetTexture2D() : nullptr;
+		if (!Texture2DRHI)
+		{
+			UE_LOG(LogImagePlateFileSequence, Warning, TEXT("Unsupported texture type encountered: Unable to update texture to fit source frame size or bitdepth."));
+			return false;
+		}
+		else
+		{
+			// At least check the bitdepth
+			uint32 DestBitDepth = GPixelFormats[Texture2DRHI->GetFormat()].BlockBytes * 8 / GPixelFormats[Texture2DRHI->GetFormat()].NumComponents;
+			if (DestBitDepth != BitDepth)
+			{
+				UE_LOG(LogImagePlateFileSequence, Warning, TEXT("Invalid destination texture bitdepth. Expected %u, encountered %u."), BitDepth, DestBitDepth);
+				return false;
+			}
 		}
 	}
 
@@ -116,15 +164,52 @@ TFuture<void> FImagePlateSourceFrame::CopyTo(UTexture2DDynamic* DestinationTextu
 	{
 		DestinationTexture->UpdateResource();
 	}
+	return true;
+}
+
+TFuture<void> FImagePlateSourceFrame::CopyTo(UTexture* DestinationTexture)
+{
+	// Copy the data to the texture via a render command
+	struct FRenderCommandData
+	{
+		TPromise<void> Promise;
+		bool bClearTexture;
+		FImagePlateSourceFrame SourceFrame;
+		UTexture* DestinationTexture;
+	};
+
+	// Shared render command data to allow us to pass the promise through to the command
+	typedef TSharedPtr<FRenderCommandData, ESPMode::ThreadSafe> FCommandDataPtr;
+	FCommandDataPtr CommandData = MakeShared<FRenderCommandData, ESPMode::ThreadSafe>();
+
+	TFuture<void> CompletionFuture = CommandData->Promise.GetFuture();
+
+	CommandData->bClearTexture = !Buffer.IsValid() || !EnsureTextureMetrics(DestinationTexture);
+	CommandData->SourceFrame = *this;
+	CommandData->DestinationTexture = DestinationTexture;
 
 	ENQUEUE_UNIQUE_RENDER_COMMAND_ONEPARAMETER(
 		CopySourceBufferToTexture,
 		FCommandDataPtr, CommandData, CommandData,
 		{
-			FTexture2DRHIRef Texture2DRHI = static_cast<FTexture2DDynamicResource*>(CommandData->DestinationTexture->Resource)->GetTexture2DRHI();
+			if (!CommandData->DestinationTexture->Resource || !CommandData->DestinationTexture->Resource->TextureRHI)
+			{
+				return;
+			}
+
+			FTexture2DRHIRef Texture2DRHI = CommandData->DestinationTexture->Resource->TextureRHI->GetTexture2D();
+			if (!Texture2DRHI)
+			{
+				return;
+			}
 
 			uint32 DestPitch = 0;
 			uint8* DestinationBuffer = (uint8*)RHILockTexture2D(Texture2DRHI, 0, RLM_WriteOnly, DestPitch, false);
+			if (!DestinationBuffer)
+			{
+				UE_LOG(LogImagePlateFileSequence, Warning, TEXT("Unable to lock texture for write"));
+				return;
+			}
 
 			if (CommandData->bClearTexture)
 			{
@@ -290,21 +375,6 @@ namespace ImagePlateFrameCache
 		: CurrentFrameNumber(-1), MinCacheRange(-1), MaxCacheRange(-1), Framerate(InFramerate)
 	{
 		FString SequenceFolder = InSequencePath;
-		if (!FPackageName::TryConvertLongPackageNameToFilename(InSequencePath, SequenceFolder))
-		{
-			SequenceFolder = InSequencePath;
-			if (FPaths::IsRelative(SequenceFolder))
-			{
-				if (SequenceFolder.StartsWith(TEXT("./")))
-				{
-					SequenceFolder = FPaths::ConvertRelativePathToFull(FPaths::ProjectContentDir(), SequenceFolder.RightChop(2));
-				}
-				else
-				{
-					SequenceFolder = FPaths::ConvertRelativePathToFull(SequenceFolder);
-				}
-			}
-		}
 
 		IFileManager::Get().FindFiles(FrameFilenames, *SequenceFolder, *InWildcard);
 		if (FrameFilenames.Num() == 0)
@@ -482,12 +552,14 @@ namespace ImagePlateFrameCache
 		const TArray<uint8>* RawImageData = nullptr;
 		if (ImageWrapper->GetRaw(ERGBFormat::RGBA, SourceBitDepth, RawImageData) && RawImageData)
 		{
+			// BMP image wrappers supply the bitdepth per pixel, rather than per-channel
+			SourceBitDepth = ImageType == EImageFormat::BMP ? ImageWrapper->GetBitDepth() / 4 : ImageWrapper->GetBitDepth();
 
 			return FImagePlateSourceFrame(
 				*RawImageData,
 				ImageWrapper->GetWidth(),
 				ImageWrapper->GetHeight(),
-				ImageWrapper->GetBitDepth()
+				SourceBitDepth
 			);
 		}
 		else
diff --git a/Engine/Plugins/Runtime/ApexDestruction/Source/ApexDestruction/Private/DestructibleComponent.cpp b/Engine/Plugins/Runtime/ApexDestruction/Source/ApexDestruction/Private/DestructibleComponent.cpp
--- a/Engine/Plugins/Runtime/ApexDestruction/Source/ApexDestruction/Private/DestructibleComponent.cpp
+++ b/Engine/Plugins/Runtime/ApexDestruction/Source/ApexDestruction/Private/DestructibleComponent.cpp
@@ -21,6 +21,7 @@
 #include "Engine/StaticMesh.h"
 #include "Physics/PhysicsFiltering.h"
 #include "ApexDestructionModule.h"
+#include "PhysicsEngine/PhysicsSettings.h"
 
 UDestructibleComponent::UDestructibleComponent(const FObjectInitializer& ObjectInitializer)
 	: Super(ObjectInitializer)
@@ -45,6 +46,9 @@ UDestructibleComponent::UDestructibleComponent(const FObjectInitializer& ObjectI
 	LargeChunkThreshold = 25.f;
 
 	SetComponentSpaceTransformsDoubleBuffering(false);
+
+	// Get contact offset params
+	FBodySetupShapeIterator::GetContactOffsetParams(ContactOffsetFactor, MinContactOffset, MaxContactOffset);
 }
 
 #if WITH_EDITORONLY_DATA
@@ -1625,6 +1629,11 @@ void UDestructibleComponent::SetCollisionResponseForShape(PxShape* Shape, int32
 
 		SCOPED_SCENE_WRITE_LOCK(Shape->getActor()->getScene());
 
+		// See FBodySetupShapeIterator::ComputeContactOffset(const PxConvexMeshGeometry&) for this calculation. 
+		// Reproduced here to avoid going to the underlying PxConvexMeshGeometry
+		const float MinBoundsExtent = Shape->getActor()->getWorldBounds().getExtents().minElement();
+		Shape->setContactOffset(FMath::Clamp(ContactOffsetFactor * MinBoundsExtent, MinContactOffset, MaxContactOffset));
+
 		Shape->setQueryFilterData(PQueryFilterData);
 		Shape->setSimulationFilterData(PSimFilterData);
 		Shape->setFlag(PxShapeFlag::eSCENE_QUERY_SHAPE, true);
diff --git a/Engine/Source/Editor/Sequencer/Private/Sequencer.cpp b/Engine/Source/Editor/Sequencer/Private/Sequencer.cpp
--- a/Engine/Source/Editor/Sequencer/Private/Sequencer.cpp
+++ b/Engine/Source/Editor/Sequencer/Private/Sequencer.cpp
@@ -104,6 +104,7 @@
 #include "LevelSequence.h"
 #include "IVREditorModule.h"
 #include "HAL/PlatformApplicationMisc.h"
+#include "SequencerKeyActor.h"
 
 #define LOCTEXT_NAMESPACE "Sequencer"
 
@@ -4072,6 +4073,12 @@ void FSequencer::SynchronizeSequencerSelectionWithExternalSelection()
 	// based on selection.
 	bool bAllAlreadySelected = true;
 
+	USelection* ActorSelection = GEditor->GetSelectedActors();
+	
+	// Get the selected sequencer keys for viewport interaction
+	TArray<ASequencerKeyActor*> SelectedSequencerKeyActors;
+	ActorSelection->GetSelectedObjects<ASequencerKeyActor>(SelectedSequencerKeyActors);
+
 	TSet<TSharedRef<FSequencerDisplayNode>> NodesToSelect;
 	for (auto ObjectBinding : NodeTree->GetObjectBindingMap() )
 	{
@@ -4086,7 +4093,21 @@ void FSequencer::SynchronizeSequencerSelectionWithExternalSelection()
 			UObject* RuntimeObject = RuntimeObjectPtr.Get();
 			if ( RuntimeObject != nullptr)
 			{
-				bool bActorSelected = GEditor->GetSelectedActors()->IsSelected( RuntimeObject );
+				for (ASequencerKeyActor* KeyActor : SelectedSequencerKeyActors)
+				{
+					if (KeyActor->IsEditorOnly())
+					{
+						AActor* TrailActor = KeyActor->GetAssociatedActor();
+						if (TrailActor != nullptr && RuntimeObject == TrailActor)
+						{
+							NodesToSelect.Add(ObjectBindingNode);
+							bAllAlreadySelected = false;
+							break;
+						}
+					}
+				}
+
+				bool bActorSelected = ActorSelection->IsSelected( RuntimeObject );
 				bool bComponentSelected = GEditor->GetSelectedComponents()->IsSelected( RuntimeObject);
 
 				if (bActorSelected || bComponentSelected)
diff --git a/Engine/Source/Runtime/ClothingSystemRuntime/Private/ClothingSimulation.cpp b/Engine/Source/Runtime/ClothingSystemRuntime/Private/ClothingSimulation.cpp
--- a/Engine/Source/Runtime/ClothingSystemRuntime/Private/ClothingSimulation.cpp
+++ b/Engine/Source/Runtime/ClothingSystemRuntime/Private/ClothingSimulation.cpp
@@ -134,42 +134,56 @@ void FClothingSimulationBase::FillContext(USkeletalMeshComponent* InComponent, f
 
 	if(USkinnedMeshComponent* MasterComponent = InComponent->MasterPoseComponent.Get())
 	{
-		const int32 NumBones = InComponent->MasterBoneMap.Num();
+		int32 NumBones = InComponent->MasterBoneMap.Num();
 
-		BaseContext->BoneTransforms.Reset(NumBones);
-		BaseContext->BoneTransforms.AddDefaulted(NumBones);
-
-		for(int32 BoneIndex = 0; BoneIndex < NumBones; ++BoneIndex)
+		if(NumBones == 0)
 		{
-			bool bFoundMaster = false;
-			if(InComponent->MasterBoneMap.IsValidIndex(BoneIndex))
+			if(InComponent->SkeletalMesh)
 			{
-				const int32 MasterIndex = InComponent->MasterBoneMap[BoneIndex];
-
-				if(MasterIndex != INDEX_NONE)
-				{
-					BaseContext->BoneTransforms[BoneIndex] = MasterComponent->GetComponentSpaceTransforms()[MasterIndex];
-					bFoundMaster = true;
-				}
-			}
+				// This case indicates an invalid master pose component (e.g. no skeletal mesh)
+				NumBones = InComponent->SkeletalMesh->RefSkeleton.GetNum();
 
-			if(!bFoundMaster)
-			{
-				if(SkelMesh)
-				{
-					const int32 ParentIndex = SkelMesh->RefSkeleton.GetParentIndex(BoneIndex);
-
-					if(ParentIndex != INDEX_NONE)
-					{
-						BaseContext->BoneTransforms[BoneIndex] = BaseContext->BoneTransforms[ParentIndex] * SkelMesh->RefSkeleton.GetRefBonePose()[BoneIndex];
-					}
-					else
-					{
-						BaseContext->BoneTransforms[BoneIndex] = SkelMesh->RefSkeleton.GetRefBonePose()[BoneIndex];
-					}
-				}
+				BaseContext->BoneTransforms.Empty(NumBones);
+				BaseContext->BoneTransforms.AddDefaulted(NumBones);
 			}
 		}
+		else
+		{
+		    BaseContext->BoneTransforms.Reset(NumBones);
+		    BaseContext->BoneTransforms.AddDefaulted(NumBones);
+    
+		    for(int32 BoneIndex = 0; BoneIndex < NumBones; ++BoneIndex)
+		    {
+			    bool bFoundMaster = false;
+			    if(InComponent->MasterBoneMap.IsValidIndex(BoneIndex))
+			    {
+				    const int32 MasterIndex = InComponent->MasterBoneMap[BoneIndex];
+    
+				    if(MasterIndex != INDEX_NONE)
+				    {
+					    BaseContext->BoneTransforms[BoneIndex] = MasterComponent->GetComponentSpaceTransforms()[MasterIndex];
+					    bFoundMaster = true;
+				    }
+			    }
+    
+			    if(!bFoundMaster)
+			    {
+				    if(SkelMesh)
+				    {
+					    const int32 ParentIndex = SkelMesh->RefSkeleton.GetParentIndex(BoneIndex);
+    
+					    if(ParentIndex != INDEX_NONE)
+					    {
+						    BaseContext->BoneTransforms[BoneIndex] = BaseContext->BoneTransforms[ParentIndex] * SkelMesh->RefSkeleton.GetRefBonePose()[BoneIndex];
+					    }
+					    else
+					    {
+						    BaseContext->BoneTransforms[BoneIndex] = SkelMesh->RefSkeleton.GetRefBonePose()[BoneIndex];
+					    }
+				    }
+			    }
+		    }
+		}
 	}
 	else
 	{
diff --git a/Engine/Source/Runtime/Engine/Classes/Kismet/KismetRenderingLibrary.h b/Engine/Source/Runtime/Engine/Classes/Kismet/KismetRenderingLibrary.h
--- a/Engine/Source/Runtime/Engine/Classes/Kismet/KismetRenderingLibrary.h
+++ b/Engine/Source/Runtime/Engine/Classes/Kismet/KismetRenderingLibrary.h
@@ -48,6 +48,13 @@ class UKismetRenderingLibrary : public UBlueprintFunctionLibrary
 	 */
 	UFUNCTION(BlueprintCallable, Category="Rendering", meta=(WorldContext="WorldContextObject"))
 	static ENGINE_API UTextureRenderTarget2D* CreateRenderTarget2D(UObject* WorldContextObject, int32 Width = 256, int32 Height = 256, bool bHDR = true);
+	
+	/**
+	 * Manually releases GPU resources of a render target. This is useful for blueprint creating a lot of render target that would
+	 * normally be released too late by the garbage collector that can be problematic on platforms that have tight GPU memory constrains.
+	 */
+	UFUNCTION(BlueprintCallable, Category="Rendering")
+	static ENGINE_API void ReleaseRenderTarget2D(UTextureRenderTarget2D* TextureRenderTarget);
 
 	/** 
 	 * Renders a quad with the material applied to the specified render target.   
diff --git a/Engine/Source/Runtime/Engine/Private/EdGraph/EdGraphPin.cpp b/Engine/Source/Runtime/Engine/Private/EdGraph/EdGraphPin.cpp
--- a/Engine/Source/Runtime/Engine/Private/EdGraph/EdGraphPin.cpp
+++ b/Engine/Source/Runtime/Engine/Private/EdGraph/EdGraphPin.cpp
@@ -740,15 +740,19 @@ bool UEdGraphPin::ExportTextItem(FString& ValueStr, int32 PortFlags) const
 
 	for (TFieldIterator<UProperty> FieldIt(FEdGraphPinType::StaticStruct()); FieldIt; ++FieldIt)
 	{
-		FString PropertyStr;
-		const uint8* PropertyAddr = FieldIt->ContainerPtrToValuePtr<uint8>(&PinType);
-		const uint8* DefaultAddr = FieldIt->ContainerPtrToValuePtr<uint8>(&DefaultPin.PinType);
-		FieldIt->ExportTextItem(PropertyStr, PropertyAddr, DefaultAddr, NULL, PortFlags, nullptr);
-
-		if (!PropertyStr.IsEmpty())
+		UProperty* Prop = *FieldIt;
+		if (Prop->ShouldPort())
 		{
-			ValueStr += PinHelpers::PinTypeName + TEXT(".") + FieldIt->GetName() + "=" + PropertyStr;
-			ValueStr += PinHelpers::ExportTextPropDelimiter;
+			FString PropertyStr;
+			const uint8* PropertyAddr = Prop->ContainerPtrToValuePtr<uint8>(&PinType);
+			const uint8* DefaultAddr = Prop->ContainerPtrToValuePtr<uint8>(&DefaultPin.PinType);
+			Prop->ExportTextItem(PropertyStr, PropertyAddr, DefaultAddr, NULL, PortFlags, nullptr);
+
+			if (!PropertyStr.IsEmpty())
+			{
+				ValueStr += PinHelpers::PinTypeName + TEXT(".") + FieldIt->GetName() + "=" + PropertyStr;
+				ValueStr += PinHelpers::ExportTextPropDelimiter;
+			}
 		}
 	}
 
@@ -1627,23 +1631,26 @@ void UEdGraphPin::SerializePinArray(FArchive& Ar, TArray<UEdGraphPin*>& ArrayRef
 	// means we're serializing for undo/redo:
 	for (UEdGraphPin* Pin : OldPins)
 	{
+		if (!Pin->WasTrashed())
+		{
 #if WITH_EDITOR
-		// More complexity to handle asymmetry in the transaction buffer. If our peer node is not in the transaction
-		// then we need to take ownership of the entire connection and clear both LinkedTo Arrays:
-		extern UNREALED_API UEditorEngine* GEditor;
-		for (int32 I = 0; I < Pin->LinkedTo.Num(); ++I )
-		{
-			UEdGraphPin* Peer = Pin->LinkedTo[I];
-			// PeerNode will be null if the pin we were linked to was already thrown away:
-			UEdGraphNode* PeerNode = Peer->GetOwningNodeUnchecked();
-			if (PeerNode && !GEditor->Trans->IsObjectTransacting(PeerNode))
+			// More complexity to handle asymmetry in the transaction buffer. If our peer node is not in the transaction
+			// then we need to take ownership of the entire connection and clear both LinkedTo Arrays:
+			extern UNREALED_API UEditorEngine* GEditor;
+			for (int32 I = 0; I < Pin->LinkedTo.Num(); ++I)
 			{
-				Pin->BreakLinkTo(Peer);
-				--I;
+				UEdGraphPin* Peer = Pin->LinkedTo[I];
+				// PeerNode will be null if the pin we were linked to was already thrown away:
+				UEdGraphNode* PeerNode = Peer->GetOwningNodeUnchecked();
+				if (PeerNode && !GEditor->Trans->IsObjectTransacting(PeerNode))
+				{
+					Pin->BreakLinkTo(Peer);
+					--I;
+				}
 			}
-		}
 #endif// WITH_EDITOR
-		Pin->DestroyImpl(false);
+			Pin->DestroyImpl(false);
+		}
 	}
 }
 
diff --git a/Engine/Source/Runtime/ShaderCore/Private/ShaderCore.cpp b/Engine/Source/Runtime/ShaderCore/Private/ShaderCore.cpp
--- a/Engine/Source/Runtime/ShaderCore/Private/ShaderCore.cpp
+++ b/Engine/Source/Runtime/ShaderCore/Private/ShaderCore.cpp
@@ -310,15 +310,27 @@ void VerifyShaderSourceFiles()
 			SlowTask.EnterProgressFrame(1);
 			FString FileContents;
 			// load each shader source file. This will cache the shader source data after it has been verified
-			LoadShaderSourceFileChecked(*VirtualShaderSourcePaths[ShaderFileIdx], FileContents);
+			LoadShaderSourceFile(*VirtualShaderSourcePaths[ShaderFileIdx], FileContents, nullptr);
 		}
 	}
 }
 
-static FString GetShaderSourceFilePath(const FString& VirtualFilePath)
+static void LogShaderSourceDirectoryMappings()
+{
+	const TMap<FString, FString>& ShaderSourceDirectoryMappings = FPlatformProcess::AllShaderSourceDirectoryMappings();
+	for (const auto& Iter : ShaderSourceDirectoryMappings)
+	{
+		UE_LOG(LogShaders, Log, TEXT("Shader directory mapping %s -> %s"), *Iter.Key, *Iter.Value);
+	}
+}
+
+static FString GetShaderSourceFilePath(const FString& VirtualFilePath, TArray<FShaderCompilerError>* CompileErrors)
 {
 	// Make sure the .usf extension is correctly set.
-	check(CheckVirtualShaderFilePath(VirtualFilePath));
+	if (!CheckVirtualShaderFilePath(VirtualFilePath, CompileErrors))
+	{
+		return FString();
+	}
 
 	// We don't cache the output of this function because only used in LoadShaderSourceFile that is cached, or when there
 	// is shader compilation errors.
@@ -344,10 +356,18 @@ static FString GetShaderSourceFilePath(const FString& VirtualFilePath)
 	}
 
 	// Make sure a directory mapping has matched.
-	checkf(!RealFilePath.IsEmpty(), TEXT("Could not map virtual shader source path \"%s\"."), *VirtualFilePath);
+	if (RealFilePath.IsEmpty())
+	{
+		FString Error = FString::Printf(TEXT("Can't map virtual shader source path \"%s\"."), *VirtualFilePath);
+		Error += TEXT("\nDirectory mappings are:");
+		for (const auto& Iter : ShaderSourceDirectoryMappings)
+		{
+			Error += FString::Printf(TEXT("\n  %s -> %s"), *Iter.Key, *Iter.Value);
+		}
+
+		ReportVirtualShaderFilePathError(CompileErrors, Error);
+	}
 	
-	// Normally, the return should be a relative path from process's base dir.
-	check(FPaths::IsRelative(RealFilePath));
 	return RealFilePath;
 }
 
@@ -406,7 +426,7 @@ FString ParseVirtualShaderFilename(const FString& InFilename)
 	return OutputFilename;
 }
 
-bool LoadShaderSourceFile(const TCHAR* VirtualFilePath, FString& OutFileContents) // TODO: const FString&
+bool LoadShaderSourceFile(const TCHAR* VirtualFilePath, FString& OutFileContents, TArray<FShaderCompilerError>* OutCompileErrors) // TODO: const FString&
 {
 	// it's not expected that cooked platforms get here, but if they do, this is the final out
 	if (FPlatformProperties::RequiresCookedData())
@@ -433,10 +453,10 @@ bool LoadShaderSourceFile(const TCHAR* VirtualFilePath, FString& OutFileContents
 		}
 		else
 		{
-			FString ShaderFilePath = GetShaderSourceFilePath(VirtualFilePath);
+			FString ShaderFilePath = GetShaderSourceFilePath(VirtualFilePath, OutCompileErrors);
 
 			// verify SHA hash of shader files on load. missing entries trigger an error
-			if (FFileHelper::LoadFileToString(OutFileContents, *ShaderFilePath, FFileHelper::EHashOptions::EnableVerify|FFileHelper::EHashOptions::ErrorMissingHash) )
+			if (!ShaderFilePath.IsEmpty() && FFileHelper::LoadFileToString(OutFileContents, *ShaderFilePath, FFileHelper::EHashOptions::EnableVerify|FFileHelper::EHashOptions::ErrorMissingHash) )
 			{
 				//update the shader file cache
 				GShaderFileCache.Add(VirtualFilePath, *OutFileContents);
@@ -451,7 +471,7 @@ bool LoadShaderSourceFile(const TCHAR* VirtualFilePath, FString& OutFileContents
 
 void LoadShaderSourceFileChecked(const TCHAR* VirtualFilePath, FString& OutFileContents)
 {
-	if (!LoadShaderSourceFile(VirtualFilePath, OutFileContents))
+	if (!LoadShaderSourceFile(VirtualFilePath, OutFileContents, nullptr))
 	{
 		UE_LOG(LogShaders, Fatal, TEXT("Couldn't find source file of virtual shader path \'%s\'"), VirtualFilePath);
 	}
@@ -480,14 +500,19 @@ const TCHAR* SkipToCharOnCurrentLine(const TCHAR* InStr, TCHAR TargetChar)
 /**
  * Recursively populates IncludeFilenames with the unique include filenames found in the shader file named Filename.
  */
-void GetShaderIncludes(const TCHAR* EntryPointVirtualFilePath, const TCHAR* VirtualFilePath, TArray<FString>& IncludeVirtualFilePaths, uint32 DepthLimit)
+static void GetShaderIncludes(const TCHAR* EntryPointVirtualFilePath, const TCHAR* VirtualFilePath, TArray<FString>& IncludeVirtualFilePaths, uint32 DepthLimit, bool AddToIncludeFile)
 {
 	FString FileContents;
-	LoadShaderSourceFileChecked(VirtualFilePath, FileContents);
+	LoadShaderSourceFile(VirtualFilePath, FileContents, nullptr);
 
 	//avoid an infinite loop with a 0 length string
 	if (FileContents.Len() > 0)
 	{
+		if (AddToIncludeFile)
+		{
+			IncludeVirtualFilePaths.Add(VirtualFilePath);
+		}
+
 		//find the first include directive
 		const TCHAR* IncludeBegin = FCString::Strstr(*FileContents, TEXT("#include "));
 
@@ -554,8 +579,7 @@ void GetShaderIncludes(const TCHAR* EntryPointVirtualFilePath, const TCHAR* Virt
 					{
 						if (!IncludeVirtualFilePaths.Contains(ExtractedIncludeFilename))
 						{
-							IncludeVirtualFilePaths.Add(ExtractedIncludeFilename);
-							GetShaderIncludes(EntryPointVirtualFilePath, *ExtractedIncludeFilename, IncludeVirtualFilePaths, DepthLimit - 1);
+							GetShaderIncludes(EntryPointVirtualFilePath, *ExtractedIncludeFilename, IncludeVirtualFilePaths, DepthLimit - 1, true);
 						}
 					}
 				}
@@ -583,6 +607,11 @@ void GetShaderIncludes(const TCHAR* EntryPointVirtualFilePath, const TCHAR* Virt
 	}
 }
 
+void GetShaderIncludes(const TCHAR* EntryPointVirtualFilePath, const TCHAR* VirtualFilePath, TArray<FString>& IncludeVirtualFilePaths, uint32 DepthLimit)
+{
+	GetShaderIncludes(EntryPointVirtualFilePath, VirtualFilePath, IncludeVirtualFilePaths, DepthLimit, false);
+}
+
 static void UpdateSingleShaderFilehash(FSHA1& InOutHashState, const TCHAR* VirtualFilePath)
 {
 	// Get the list of includes this file contains
@@ -730,6 +759,10 @@ void BuildShaderFileToUniformBufferMap(TMap<FString, TArray<const TCHAR*> >& Sha
 
 void InitializeShaderTypes()
 {
+	UE_LOG(LogShaders, Log, TEXT("InitializeShaderTypes() begin"));
+
+	LogShaderSourceDirectoryMappings();
+
 	TMap<FString, TArray<const TCHAR*> > ShaderFileToUniformBufferVariables;
 	BuildShaderFileToUniformBufferMap(ShaderFileToUniformBufferVariables);
 
@@ -737,14 +770,20 @@ void InitializeShaderTypes()
 	FVertexFactoryType::Initialize(ShaderFileToUniformBufferVariables);
 
 	FShaderPipelineType::Initialize();
+
+	UE_LOG(LogShaders, Log, TEXT("InitializeShaderTypes() end"));
 }
 
 void UninitializeShaderTypes()
 {
+	UE_LOG(LogShaders, Log, TEXT("UninitializeShaderTypes() begin"));
+
 	FShaderPipelineType::Uninitialize();
 
 	FShaderType::Uninitialize();
 	FVertexFactoryType::Uninitialize();
+
+	UE_LOG(LogShaders, Log, TEXT("UninitializeShaderTypes() end"));
 }
 
 /**
@@ -753,11 +792,15 @@ void UninitializeShaderTypes()
  */
 void FlushShaderFileCache()
 {
+	UE_LOG(LogShaders, Log, TEXT("FlushShaderFileCache() begin"));
+
 	GShaderHashCache.Empty();
 	GShaderFileCache.Empty();
 
 	if (!FPlatformProperties::RequiresCookedData())
 	{
+		LogShaderSourceDirectoryMappings();
+
 		TMap<FString, TArray<const TCHAR*> > ShaderFileToUniformBufferVariables;
 		BuildShaderFileToUniformBufferMap(ShaderFileToUniformBufferVariables);
 
@@ -780,6 +823,8 @@ void FlushShaderFileCache()
 			It->FlushShaderFileCache(ShaderFileToUniformBufferVariables);
 		}
 	}
+
+	UE_LOG(LogShaders, Log, TEXT("FlushShaderFileCache() end"));
 }
 
 void GenerateReferencedUniformBuffers(
@@ -843,6 +888,6 @@ FString FShaderCompilerError::GetShaderSourceFilePath() const
 	}
 	else
 	{
-	return ::GetShaderSourceFilePath(ErrorVirtualFilePath);
+		return ::GetShaderSourceFilePath(ErrorVirtualFilePath, nullptr);
 	}
 }
