Fix non-monolithic build error

diff --git a/Engine/Plugins/Experimental/ImagePlate/Source/ImagePlateEditor/Private/ImagePlateEditor.cpp b/Engine/Plugins/Experimental/ImagePlate/Source/ImagePlateEditor/Private/ImagePlateEditor.cpp
--- a/Engine/Plugins/Experimental/ImagePlate/Source/ImagePlateEditor/Private/ImagePlateEditor.cpp
+++ b/Engine/Plugins/Experimental/ImagePlate/Source/ImagePlateEditor/Private/ImagePlateEditor.cpp
@@ -23,8 +23,7 @@ public:
 		const FVector2D Icon16x16(16.f, 16.f);
 		const FVector2D Icon64x64(64.f, 64.f);
 
-		FString PluginBasePath = FPaths::GetPath(FModuleManager::Get().GetModuleFilename("ImagePlate"));
-		SetContentRoot(PluginBasePath / TEXT("../../Resources/Icons"));
+		SetContentRoot(FPaths::EnginePluginsDir() / TEXT("../../Resources/Icons"));
 
 		Set("ClassIcon.ImagePlate", new FSlateImageBrush(RootToContentDir(TEXT("ImagePlate_16x.png")), Icon16x16));
 		Set("ClassThumbnail.ImagePlate", new FSlateImageBrush(RootToContentDir(TEXT("ImagePlate_64x.png")), Icon64x64));
