Merging //UE4/Dev-Main @ 4051136 to Dev-Core (//UE4/Dev-Core)

diff --git a/Engine/Plugins/Experimental/ImagePlate/Source/ImagePlateEditor/Private/ImagePlateEditor.cpp b/Engine/Plugins/Experimental/ImagePlate/Source/ImagePlateEditor/Private/ImagePlateEditor.cpp
--- a/Engine/Plugins/Experimental/ImagePlate/Source/ImagePlateEditor/Private/ImagePlateEditor.cpp
+++ b/Engine/Plugins/Experimental/ImagePlate/Source/ImagePlateEditor/Private/ImagePlateEditor.cpp
@@ -23,8 +23,7 @@ public:
 		const FVector2D Icon16x16(16.f, 16.f);
 		const FVector2D Icon64x64(64.f, 64.f);
 
-		FString PluginBasePath = FPaths::GetPath(FModuleManager::Get().GetModuleFilename("ImagePlate"));
-		SetContentRoot(PluginBasePath / TEXT("../../Resources/Icons"));
+		SetContentRoot(FPaths::EnginePluginsDir() / TEXT("../../Resources/Icons"));
 
 		Set("ClassIcon.ImagePlate", new FSlateImageBrush(RootToContentDir(TEXT("ImagePlate_16x.png")), Icon16x16));
 		Set("ClassThumbnail.ImagePlate", new FSlateImageBrush(RootToContentDir(TEXT("ImagePlate_64x.png")), Icon64x64));
diff --git a/Engine/Source/Runtime/Engine/Private/Materials/MaterialInterface.cpp b/Engine/Source/Runtime/Engine/Private/Materials/MaterialInterface.cpp
--- a/Engine/Source/Runtime/Engine/Private/Materials/MaterialInterface.cpp
+++ b/Engine/Source/Runtime/Engine/Private/Materials/MaterialInterface.cpp
@@ -106,6 +106,18 @@ FMaterialRelevance UMaterialInterface::GetRelevance_Internal(const UMaterial* Ma
 	if(Material)
 	{
 		const FMaterialResource* MaterialResource = GetMaterialResource(InFeatureLevel);
+
+		// If material is invalid e.g. unparented instance, fallback to the passed in material
+		if (!MaterialResource && Material)
+		{
+			MaterialResource = Material->GetMaterialResource(InFeatureLevel);	
+		}
+
+		if (!MaterialResource)
+		{
+			return FMaterialRelevance();
+		}
+
 		const EBlendMode BlendMode = (EBlendMode)GetBlendMode();
 		const bool bIsTranslucent = IsTranslucentBlendMode(BlendMode);
 
