Framework - Making the code needed to get the MinimalView info out of a scene capture component easier and placing the code to more easliy calculate the view/projection matricies into an area that's easier to call for gameplay code that needs that information.

Michael.Noland, Marc.Audy

diff --git a/Engine/Plugins/Experimental/ImagePlate/Source/ImagePlate/Private/ImagePlateComponent.cpp b/Engine/Plugins/Experimental/ImagePlate/Source/ImagePlate/Private/ImagePlateComponent.cpp
--- a/Engine/Plugins/Experimental/ImagePlate/Source/ImagePlate/Private/ImagePlateComponent.cpp
+++ b/Engine/Plugins/Experimental/ImagePlate/Source/ImagePlate/Private/ImagePlateComponent.cpp
@@ -30,56 +30,10 @@
 #include "ImagePlateFrustumComponent.h"
 #include "ImagePlateComponent.h"
 #include "StaticMeshResources.h"
+#include "Kismet/GameplayStatics.h"
 
 namespace
 {
-	FMatrix CalculateProjectionMatrix(const FMinimalViewInfo& MinimalView)
-	{
-		FMatrix ProjectionMatrix;
-
-		if (MinimalView.ProjectionMode == ECameraProjectionMode::Orthographic)
-		{
-			const float YScale = 1.0f / MinimalView.AspectRatio;
-
-			const float HalfOrthoWidth = MinimalView.OrthoWidth / 2.0f;
-			const float ScaledOrthoHeight = MinimalView.OrthoWidth / 2.0f * YScale;
-
-			const float NearPlane = MinimalView.OrthoNearClipPlane;
-			const float FarPlane = MinimalView.OrthoFarClipPlane;
-
-			const float ZScale = 1.0f / (FarPlane - NearPlane);
-			const float ZOffset = -NearPlane;
-
-			ProjectionMatrix = FReversedZOrthoMatrix(
-				HalfOrthoWidth,
-				ScaledOrthoHeight,
-				ZScale,
-				ZOffset
-				);
-		}
-		else
-		{
-			// Avoid divide by zero in the projection matrix calculation by clamping FOV
-			ProjectionMatrix = FReversedZPerspectiveMatrix(
-				FMath::Max(0.001f, MinimalView.FOV) * (float)PI / 360.0f,
-				MinimalView.AspectRatio,
-				1.0f,
-				GNearClippingPlane );
-		}
-
-		if (!MinimalView.OffCenterProjectionOffset.IsZero())
-		{
-			const float Left = -1.0f + MinimalView.OffCenterProjectionOffset.X;
-			const float Right = Left + 2.0f;
-			const float Bottom = -1.0f + MinimalView.OffCenterProjectionOffset.Y;
-			const float Top = Bottom + 2.0f;
-			ProjectionMatrix.M[2][0] = (Left + Right) / (Left - Right);
-			ProjectionMatrix.M[2][1] = (Bottom + Top) / (Bottom - Top);
-		}
-
-		return ProjectionMatrix;
-	}
-
 	class FImagePlateIndexBuffer : public FIndexBuffer
 	{
 	public:
@@ -372,11 +326,14 @@ void UImagePlateComponent::UpdateTransformScale()
 	AActor* ViewTarget = FindViewTarget();
 	if (ViewTarget && Plate.bFillScreen)
 	{
-		GetProjectionMatricesFromViewTarget(ViewTarget, ViewProjectionMatrix, InvViewProjectionMatrix);
+		FMatrix ViewMatrix;
+		FMatrix ProjectionMatrix;
+		FMatrix ViewProjectionMatrix;
+		UGameplayStatics::CalculateViewProjectionMatricesFromViewTarget(ViewTarget, ViewMatrix, ProjectionMatrix, ViewProjectionMatrix);
 
 		const FMatrix LocalToWorld = GetComponentTransform().ToMatrixNoScale();
 		const FMatrix WorldToLocal = LocalToWorld.Inverse();
-		const FMatrix ScreenToLocalSpace = InvViewProjectionMatrix * WorldToLocal;
+		const FMatrix ScreenToLocalSpace = ViewProjectionMatrix.Inverse() * WorldToLocal;
 
 
 		// Just use the current view projection matrix
@@ -454,54 +411,6 @@ void UImagePlateComponent::GetUsedMaterials(TArray<UMaterialInterface*>& OutMate
 	OutMaterials.AddUnique(Plate.DynamicMaterial ? Plate.DynamicMaterial : Plate.Material);
 }
 
-void UImagePlateComponent::GetProjectionMatricesFromViewTarget(AActor* InViewTarget, FMatrix& OutViewProjectionMatrix, FMatrix& OutInvViewProjectionMatrix)
-{
-	check(InViewTarget);
-
-	FMinimalViewInfo MinimalViewInfo;
-	
-	ASceneCapture2D* SceneCapture = Cast<ASceneCapture2D>(InViewTarget);
-	USceneCaptureComponent2D* SceneCaptureComponent = SceneCapture ? SceneCapture->GetCaptureComponent2D() : nullptr;
-	if (SceneCaptureComponent)
-	{
-		MinimalViewInfo.Location = SceneCaptureComponent->GetComponentLocation();
-		MinimalViewInfo.Rotation = SceneCaptureComponent->GetComponentRotation();
-
-		MinimalViewInfo.FOV = SceneCaptureComponent->FOVAngle;
-		MinimalViewInfo.AspectRatio = SceneCaptureComponent->TextureTarget ? float(SceneCaptureComponent->TextureTarget->SizeX) / SceneCaptureComponent->TextureTarget->SizeY : 1.f;
-		MinimalViewInfo.bConstrainAspectRatio = false;
-		MinimalViewInfo.ProjectionMode = SceneCaptureComponent->ProjectionType;
-		MinimalViewInfo.OrthoWidth = SceneCaptureComponent->OrthoWidth;
-	}
-	else
-	{
-		InViewTarget->CalcCamera(0.f, MinimalViewInfo);
-	}
-
-	FMatrix ViewRotationMatrix = FInverseRotationMatrix(MinimalViewInfo.Rotation) * FMatrix(
-		FPlane(0, 0, 1, 0),
-		FPlane(1, 0, 0, 0),
-		FPlane(0, 1, 0, 0),
-		FPlane(0, 0, 0, 1));
-
-	FMatrix ProjectionMatrix;
-	if (SceneCaptureComponent && SceneCaptureComponent->bUseCustomProjectionMatrix)
-	{
-		ProjectionMatrix = AdjustProjectionMatrixForRHI(SceneCaptureComponent->CustomProjectionMatrix);
-	}
-	else
-	{
-		ProjectionMatrix = AdjustProjectionMatrixForRHI(CalculateProjectionMatrix(MinimalViewInfo));
-	}
-
-	FMatrix ViewMatrix = FTranslationMatrix(-MinimalViewInfo.Location) * ViewRotationMatrix;
-	FMatrix InvProjectionMatrix = ProjectionMatrix.Inverse();
-	FMatrix InvViewMatrix = ViewRotationMatrix.GetTransposed() * FTranslationMatrix(MinimalViewInfo.Location);
-
-	OutViewProjectionMatrix = ViewMatrix * ProjectionMatrix;
-	OutInvViewProjectionMatrix = InvProjectionMatrix * InvViewMatrix;
-}
-
 FBoxSphereBounds UImagePlateComponent::CalcBounds(const FTransform& LocalToWorld) const
 {
 	return FBoxSphereBounds(FVector(0,0,0), FVector(1,  1,  1), 1.73205f).TransformBy(LocalToWorld);
diff --git a/Engine/Plugins/Experimental/ImagePlate/Source/ImagePlate/Public/ImagePlateComponent.h b/Engine/Plugins/Experimental/ImagePlate/Source/ImagePlate/Public/ImagePlateComponent.h
--- a/Engine/Plugins/Experimental/ImagePlate/Source/ImagePlate/Public/ImagePlateComponent.h
+++ b/Engine/Plugins/Experimental/ImagePlate/Source/ImagePlate/Public/ImagePlateComponent.h
@@ -122,11 +122,6 @@ public:
 	 */
 	AActor* FindViewTarget() const;
 
-	/**
-	 * Calculate projection matrices from a specified view target
-	 */
-	static void GetProjectionMatricesFromViewTarget(AActor* InViewTarget, FMatrix& OutViewProjectionMatrix, FMatrix& OutInvViewProjectionMatrix);
-
 protected:
 
 	void UpdateMaterialParametersForMedia();
diff --git a/Engine/Source/Runtime/Engine/Classes/Components/SceneCaptureComponent2D.h b/Engine/Source/Runtime/Engine/Classes/Components/SceneCaptureComponent2D.h
--- a/Engine/Source/Runtime/Engine/Classes/Components/SceneCaptureComponent2D.h
+++ b/Engine/Source/Runtime/Engine/Classes/Components/SceneCaptureComponent2D.h
@@ -111,6 +111,8 @@ class ENGINE_API USceneCaptureComponent2D : public USceneCaptureComponent
 
 	//~ End UObject Interface
 
+	virtual void GetCameraView(float DeltaTime, FMinimalViewInfo& DesiredView);
+
 	/** Adds an Blendable (implements IBlendableInterface) to the array of Blendables (if it doesn't exist) and update the weight */
 	UFUNCTION(BlueprintCallable, Category="Rendering")
 	void AddOrUpdateBlendable(TScriptInterface<IBlendableInterface> InBlendableObject, float InWeight = 1.0f) { PostProcessSettings.AddBlendable(InBlendableObject, InWeight); }
diff --git a/Engine/Source/Runtime/Engine/Classes/Engine/SceneCapture2D.h b/Engine/Source/Runtime/Engine/Classes/Engine/SceneCapture2D.h
--- a/Engine/Source/Runtime/Engine/Classes/Engine/SceneCapture2D.h
+++ b/Engine/Source/Runtime/Engine/Classes/Engine/SceneCapture2D.h
@@ -26,9 +26,8 @@ public:
 	UFUNCTION(BlueprintCallable, Category="Rendering")
 	void OnInterpToggle(bool bEnable);
 
+	virtual void CalcCamera(float DeltaTime, struct FMinimalViewInfo& OutMinimalViewInfo) override;
+
 	/** Returns CaptureComponent2D subobject **/
 	ENGINE_API class USceneCaptureComponent2D* GetCaptureComponent2D() const { return CaptureComponent2D; }
 };
-
-
-
diff --git a/Engine/Source/Runtime/Engine/Classes/Kismet/GameplayStatics.h b/Engine/Source/Runtime/Engine/Classes/Kismet/GameplayStatics.h
--- a/Engine/Source/Runtime/Engine/Classes/Kismet/GameplayStatics.h
+++ b/Engine/Source/Runtime/Engine/Classes/Kismet/GameplayStatics.h
@@ -1082,6 +1082,16 @@ public:
 	UFUNCTION(BlueprintPure, Category = "Utilities")
 	static bool ProjectWorldToScreen(APlayerController const* Player, const FVector& WorldPosition, FVector2D& ScreenPosition, bool bPlayerViewportRelative = false);
 
+	/**
+	 * Calculate view-projection matrices from a specified view target
+	 */
+	static void CalculateViewProjectionMatricesFromViewTarget(AActor* InViewTarget, FMatrix& OutViewMatrix, FMatrix& OutProjectionMatrix, FMatrix& OutViewProjectionMatrix);
+
+	/**
+	 * Calculate view-projection matrices from a specified MinimalViewInfo and optional custom projection matrix
+	 */
+	static void CalculateViewProjectionMatricesFromMinimalView(const FMinimalViewInfo& MinimalViewInfo, const TOptional<FMatrix>& CustomProjectionMatrix, FMatrix& OutViewMatrix, FMatrix& OutProjectionMatrix, FMatrix& OutViewProjectionMatrix);
+
 	//~ Utility functions for interacting with Options strings
 
 	//~=========================================================================
diff --git a/Engine/Source/Runtime/Engine/Private/Components/SceneCaptureComponent.cpp b/Engine/Source/Runtime/Engine/Private/Components/SceneCaptureComponent.cpp
--- a/Engine/Source/Runtime/Engine/Private/Components/SceneCaptureComponent.cpp
+++ b/Engine/Source/Runtime/Engine/Private/Components/SceneCaptureComponent.cpp
@@ -31,6 +31,7 @@
 #include "Engine/BlueprintGeneratedClass.h"
 #include "Engine/SimpleConstructionScript.h"
 #include "Engine/SCS_Node.h"
+#include "Engine/TextureRenderTarget2D.h"
 
 #define LOCTEXT_NAMESPACE "SceneCaptureComponent"
 
@@ -96,6 +97,18 @@ void ASceneCapture2D::OnInterpToggle(bool bEnable)
 {
 	CaptureComponent2D->SetVisibility(bEnable);
 }
+
+void ASceneCapture2D::CalcCamera(float DeltaTime, FMinimalViewInfo& OutMinimalViewInfo)
+{
+	if (USceneCaptureComponent2D* SceneCaptureComponent = GetCaptureComponent2D())
+	{
+		SceneCaptureComponent->GetCameraView(DeltaTime, OutMinimalViewInfo);
+	}
+	else
+	{
+		Super::CalcCamera(DeltaTime, OutMinimalViewInfo);
+	}
+}
 // -----------------------------------------------
 
 ASceneCaptureCube::ASceneCaptureCube(const FObjectInitializer& ObjectInitializer)
@@ -501,6 +514,18 @@ void USceneCaptureComponent2D::TickComponent(float DeltaTime, enum ELevelTick Ti
 	}
 }
 
+void USceneCaptureComponent2D::GetCameraView(float DeltaTime, FMinimalViewInfo& OutMinimalViewInfo)
+{
+	OutMinimalViewInfo.Location = GetComponentLocation();
+	OutMinimalViewInfo.Rotation = GetComponentRotation();
+	
+	OutMinimalViewInfo.FOV = FOVAngle;
+	OutMinimalViewInfo.AspectRatio = TextureTarget ? (float(TextureTarget->SizeX) / TextureTarget->SizeY) : 1.f;
+	OutMinimalViewInfo.bConstrainAspectRatio = false;
+	OutMinimalViewInfo.ProjectionMode = ProjectionType;
+	OutMinimalViewInfo.OrthoWidth = OrthoWidth;
+}
+
 void USceneCaptureComponent2D::CaptureSceneDeferred()
 {
 	UWorld* World = GetWorld();
diff --git a/Engine/Source/Runtime/Engine/Private/GameplayStatics.cpp b/Engine/Source/Runtime/Engine/Private/GameplayStatics.cpp
--- a/Engine/Source/Runtime/Engine/Private/GameplayStatics.cpp
+++ b/Engine/Source/Runtime/Engine/Private/GameplayStatics.cpp
@@ -41,6 +41,8 @@
 #include "PhysicsEngine/BodySetup.h"
 #include "Misc/EngineVersion.h"
 #include "ContentStreaming.h"
+#include "Engine/SceneCapture2D.h"
+#include "Components/SceneCaptureComponent2D.h"
 
 #define LOCTEXT_NAMESPACE "GameplayStatics"
 
@@ -2551,6 +2553,55 @@ bool UGameplayStatics::ProjectWorldToScreen(APlayerController const* Player, con
 	return false;
 }
 
+void UGameplayStatics::CalculateViewProjectionMatricesFromViewTarget(AActor* InViewTarget, FMatrix& OutViewMatrix, FMatrix& OutProjectionMatrix, FMatrix& OutViewProjectionMatrix)
+{
+	if (InViewTarget)
+	{
+		FMinimalViewInfo MinimalViewInfo;
+		InViewTarget->CalcCamera(0.f, MinimalViewInfo);
+
+		// This is kinda weird odd one-off, can this be integrated into the MinimalViewInfo?
+		TOptional<FMatrix> CustomProjectionMatrix;
+		if (ASceneCapture2D* SceneCapture2D = Cast<ASceneCapture2D>(InViewTarget))
+		{
+			if (USceneCaptureComponent2D* SceneCaptureComponent2D = SceneCapture2D->GetCaptureComponent2D())
+			{
+				if (SceneCaptureComponent2D && SceneCaptureComponent2D->bUseCustomProjectionMatrix)
+				{
+					CustomProjectionMatrix = SceneCaptureComponent2D->CustomProjectionMatrix;
+				}
+			}
+		}
+
+		CalculateViewProjectionMatricesFromMinimalView(MinimalViewInfo, CustomProjectionMatrix, OutViewMatrix, OutProjectionMatrix, OutViewProjectionMatrix);
+	}
+}
+
+void UGameplayStatics::CalculateViewProjectionMatricesFromMinimalView(const FMinimalViewInfo& MinimalViewInfo, const TOptional<FMatrix>& CustomProjectionMatrix, FMatrix& OutViewMatrix, FMatrix& OutProjectionMatrix, FMatrix& OutViewProjectionMatrix)
+{
+	if (CustomProjectionMatrix.IsSet())
+	{
+		OutProjectionMatrix = AdjustProjectionMatrixForRHI(CustomProjectionMatrix.GetValue());
+	}
+	else
+	{
+		OutProjectionMatrix = AdjustProjectionMatrixForRHI(MinimalViewInfo.CalculateProjectionMatrix());
+	}
+
+	FMatrix ViewRotationMatrix = FInverseRotationMatrix(MinimalViewInfo.Rotation) * FMatrix(
+		FPlane(0, 0, 1, 0),
+		FPlane(1, 0, 0, 0),
+		FPlane(0, 1, 0, 0),
+		FPlane(0, 0, 0, 1));
+
+	OutViewMatrix = FTranslationMatrix(-MinimalViewInfo.Location) * ViewRotationMatrix;
+	//OutInvProjectionMatrix = OutProjectionMatrix.Inverse();
+	//OutInvViewMatrix = ViewRotationMatrix.GetTransposed() * FTranslationMatrix(MinimalViewInfo.Location);
+
+	OutViewProjectionMatrix = OutViewMatrix * OutProjectionMatrix;
+	//OutInvViewProjectionMatrix = OutInvProjectionMatrix * OutInvViewMatrix;
+}
+
 bool UGameplayStatics::GrabOption( FString& Options, FString& Result )
 {
 	FString QuestionMark(TEXT("?"));
@@ -2641,4 +2692,9 @@ bool UGameplayStatics::HasLaunchOption(const FString& OptionToCheck)
 	return FParse::Param(FCommandLine::Get(), *OptionToCheck);
 }
 
+/**
+ * Calculate projection matrices from a specified view target
+ */
+static void GetProjectionMatricesFromViewTarget(AActor* InViewTarget, FMatrix& OutViewProjectionMatrix, FMatrix& OutInvViewProjectionMatrix);
+
 #undef LOCTEXT_NAMESPACE
