Fixing Linux build.

diff --git a/Engine/Plugins/Experimental/ImagePlate/Source/ImagePlate/Private/ImagePlateComponent.cpp b/Engine/Plugins/Experimental/ImagePlate/Source/ImagePlate/Private/ImagePlateComponent.cpp
--- a/Engine/Plugins/Experimental/ImagePlate/Source/ImagePlate/Private/ImagePlateComponent.cpp
+++ b/Engine/Plugins/Experimental/ImagePlate/Source/ImagePlate/Private/ImagePlateComponent.cpp
@@ -328,8 +328,8 @@ void UImagePlateComponent::UpdateTransformScale()
 	{
 		FMatrix ViewMatrix;
 		FMatrix ProjectionMatrix;
-		FMatrix ViewProjectionMatrix;
 		UGameplayStatics::CalculateViewProjectionMatricesFromViewTarget(ViewTarget, ViewMatrix, ProjectionMatrix, ViewProjectionMatrix);
+		InvViewProjectionMatrix = ViewProjectionMatrix.Inverse();
 
 		const FMatrix LocalToWorld = GetComponentTransform().ToMatrixNoScale();
 		const FMatrix WorldToLocal = LocalToWorld.Inverse();
