Image Plate: Remove obsolete Image Plate Track/Section. Everything should be done through the Media Track now

diff --git a/Engine/Plugins/Experimental/ImagePlate/Source/ImagePlate/ImagePlate.Build.cs b/Engine/Plugins/Experimental/ImagePlate/Source/ImagePlate/ImagePlate.Build.cs
--- a/Engine/Plugins/Experimental/ImagePlate/Source/ImagePlate/ImagePlate.Build.cs
+++ b/Engine/Plugins/Experimental/ImagePlate/Source/ImagePlate/ImagePlate.Build.cs
@@ -12,8 +12,6 @@ namespace UnrealBuildTool.Rules
 					"Core",
 					"CoreUObject",
 					"Engine",
-					"MovieScene",
-					"MovieSceneTracks",
 				}
 			);
 
diff --git a/Engine/Plugins/Experimental/ImagePlate/Source/ImagePlate/Private/MovieSceneImagePlateSection.cpp b/Engine/Plugins/Experimental/ImagePlate/Source/ImagePlate/Private/MovieSceneImagePlateSection.cpp
deleted file mode 100644
diff --git a/Engine/Plugins/Experimental/ImagePlate/Source/ImagePlate/Private/MovieSceneImagePlateTemplate.cpp b/Engine/Plugins/Experimental/ImagePlate/Source/ImagePlate/Private/MovieSceneImagePlateTemplate.cpp
deleted file mode 100644
diff --git a/Engine/Plugins/Experimental/ImagePlate/Source/ImagePlate/Private/MovieSceneImagePlateTemplate.h b/Engine/Plugins/Experimental/ImagePlate/Source/ImagePlate/Private/MovieSceneImagePlateTemplate.h
deleted file mode 100644
diff --git a/Engine/Plugins/Experimental/ImagePlate/Source/ImagePlate/Private/MovieSceneImagePlateTrack.cpp b/Engine/Plugins/Experimental/ImagePlate/Source/ImagePlate/Private/MovieSceneImagePlateTrack.cpp
deleted file mode 100644
diff --git a/Engine/Plugins/Experimental/ImagePlate/Source/ImagePlate/Public/MovieSceneImagePlateSection.h b/Engine/Plugins/Experimental/ImagePlate/Source/ImagePlate/Public/MovieSceneImagePlateSection.h
deleted file mode 100644
diff --git a/Engine/Plugins/Experimental/ImagePlate/Source/ImagePlate/Public/MovieSceneImagePlateTrack.h b/Engine/Plugins/Experimental/ImagePlate/Source/ImagePlate/Public/MovieSceneImagePlateTrack.h
deleted file mode 100644
diff --git a/Engine/Plugins/Experimental/ImagePlate/Source/ImagePlateEditor/ImagePlateEditor.Build.cs b/Engine/Plugins/Experimental/ImagePlate/Source/ImagePlateEditor/ImagePlateEditor.Build.cs
--- a/Engine/Plugins/Experimental/ImagePlate/Source/ImagePlateEditor/ImagePlateEditor.Build.cs
+++ b/Engine/Plugins/Experimental/ImagePlate/Source/ImagePlateEditor/ImagePlateEditor.Build.cs
@@ -16,12 +16,7 @@ namespace UnrealBuildTool.Rules
 					"EditorStyle",
 					"Engine",
 					"ImagePlate",
-					"LevelSequence",
-					"MovieScene",
-					"MovieSceneTools",
-					"MovieSceneTracks",
 					"RHI",
-					"Sequencer",
 					"Slate",
 					"SlateCore",
                     "TimeManagement",
diff --git a/Engine/Plugins/Experimental/ImagePlate/Source/ImagePlateEditor/Private/ImagePlateEditor.cpp b/Engine/Plugins/Experimental/ImagePlate/Source/ImagePlateEditor/Private/ImagePlateEditor.cpp
--- a/Engine/Plugins/Experimental/ImagePlate/Source/ImagePlateEditor/Private/ImagePlateEditor.cpp
+++ b/Engine/Plugins/Experimental/ImagePlate/Source/ImagePlateEditor/Private/ImagePlateEditor.cpp
@@ -1,13 +1,10 @@
 // Copyright Epic Games, Inc. All Rights Reserved.
 
-#include "ISequencerModule.h"
 #include "Modules/ModuleInterface.h"
 #include "Modules/ModuleManager.h"
-#include "MovieSceneImagePlateSection.h"
 #include "PropertyEditorModule.h"
 #include "Styling/SlateStyle.h"
 #include "Styling/SlateStyleRegistry.h"
-#include "ImagePlateTrackEditor.h"
 #include "AssetTypeActions_ImagePlateFileSequence.h"
 #include "IAssetTools.h"
 #include "ISettingsModule.h"
@@ -68,10 +65,6 @@ public:
 	virtual void StartupModule() override
 	{
 		FImagePlateEditorStyle::Get();
-
-		ISequencerModule& SequencerModule = FModuleManager::Get().LoadModuleChecked<ISequencerModule>("Sequencer");
-		TrackEditorBindingHandle = SequencerModule.RegisterPropertyTrackEditor<FImagePlateTrackEditor>();
-
 		AssetTypeActions = MakeShared<FAssetTypeActions_ImagePlateFileSequence>();
 
 		IAssetTools& AssetToolsModule = FModuleManager::Get().LoadModuleChecked<FAssetToolsModule>("AssetTools").Get();
@@ -89,12 +82,6 @@ public:
 	{
 		FImagePlateEditorStyle::Destroy();
 
-		ISequencerModule* SequencerModulePtr = FModuleManager::Get().GetModulePtr<ISequencerModule>("Sequencer");
-		if (SequencerModulePtr)
-		{
-			SequencerModulePtr->UnRegisterTrackEditor(TrackEditorBindingHandle);
-		}
-
 		FAssetToolsModule* AssetToolsModule = FModuleManager::Get().GetModulePtr<FAssetToolsModule>("AssetTools");
 		if (AssetToolsModule)
 		{
@@ -110,7 +97,6 @@ public:
 	}
 
 	TSharedPtr<FAssetTypeActions_ImagePlateFileSequence> AssetTypeActions;
-	FDelegateHandle TrackEditorBindingHandle;
 };
 
 IMPLEMENT_MODULE(FImagePlateEditorModule, ImagePlateEditor);
diff --git a/Engine/Plugins/Experimental/ImagePlate/Source/ImagePlateEditor/Private/ImagePlateTrackEditor.cpp b/Engine/Plugins/Experimental/ImagePlate/Source/ImagePlateEditor/Private/ImagePlateTrackEditor.cpp
deleted file mode 100644
diff --git a/Engine/Plugins/Experimental/ImagePlate/Source/ImagePlateEditor/Private/ImagePlateTrackEditor.h b/Engine/Plugins/Experimental/ImagePlate/Source/ImagePlateEditor/Private/ImagePlateTrackEditor.h
deleted file mode 100644
