Implement Ray Tracing support for Image Plate plugin

This allows the imageplate component to be visible to the path tracer.

diff --git a/Engine/Plugins/Experimental/ImagePlate/Source/ImagePlate/Private/ImagePlateComponent.cpp b/Engine/Plugins/Experimental/ImagePlate/Source/ImagePlate/Private/ImagePlateComponent.cpp
--- a/Engine/Plugins/Experimental/ImagePlate/Source/ImagePlate/Private/ImagePlateComponent.cpp
+++ b/Engine/Plugins/Experimental/ImagePlate/Source/ImagePlate/Private/ImagePlateComponent.cpp
@@ -34,6 +34,13 @@
 
 namespace
 {
+
+	static TAutoConsoleVariable<int32> CVarRayTracingImagePlate(
+		TEXT("r.RayTracing.Geometry.ImagePlate"),
+		1,
+		TEXT("Include image plate meshes in ray tracing effects (default = 1 (image plates enabled in ray tracing))"));
+
+
 	class FImagePlateIndexBuffer : public FIndexBuffer
 	{
 	public:
@@ -99,12 +106,44 @@ namespace
 			VertexBuffers.ColorVertexBuffer.ReleaseResource();
 			IndexBuffer.ReleaseResource();
 			VertexFactory.ReleaseResource();
+#if RHI_RAYTRACING
+			if (IsRayTracingEnabled())
+			{
+				RayTracingGeometry.ReleaseResource();
+			}
+#endif
 		}
 
 		virtual void CreateRenderThreadResources() override
 		{
 			BuildMesh();
 			IndexBuffer.InitResource();
+#if RHI_RAYTRACING
+			if (IsRayTracingEnabled())
+			{
+				FRayTracingGeometryInitializer Initializer;
+				Initializer.DebugName = GetOwnerName();
+				Initializer.IndexBuffer = nullptr;
+				Initializer.TotalPrimitiveCount = 0;
+				Initializer.GeometryType = RTGT_Triangles;
+				Initializer.bFastBuild = true;
+				Initializer.bAllowUpdate = false;
+
+				RayTracingGeometry.SetInitializer(Initializer);
+				RayTracingGeometry.InitResource();
+
+				RayTracingGeometry.Initializer.IndexBuffer = IndexBuffer.IndexBufferRHI;
+				RayTracingGeometry.Initializer.TotalPrimitiveCount = IndexBuffer.Indices.Num() / 3;
+
+				FRayTracingGeometrySegment Segment;
+				Segment.VertexBuffer = VertexBuffers.PositionVertexBuffer.VertexBufferRHI;
+				Segment.NumPrimitives = RayTracingGeometry.Initializer.TotalPrimitiveCount;
+				Segment.MaxVertices = VertexBuffers.PositionVertexBuffer.GetNumVertices();
+				RayTracingGeometry.Initializer.Segments.Add(Segment);
+
+				RayTracingGeometry.UpdateRHI();
+			}
+#endif
 		}
 
 		void BuildMesh()
@@ -241,12 +280,77 @@ namespace
 		virtual uint32 GetMemoryFootprint() const override { return sizeof(*this) + GetAllocatedSize(); }
 		uint32 GetAllocatedSize() const { return FPrimitiveSceneProxy::GetAllocatedSize(); }
 
+#if RHI_RAYTRACING
+		virtual bool IsRayTracingRelevant() const override { return true; }
+
+		virtual bool HasRayTracingRepresentation() const override { return true; }
+
+		virtual void GetDynamicRayTracingInstances(FRayTracingMaterialGatheringContext& Context, TArray<FRayTracingInstance>& OutRayTracingInstances) override final
+		{
+			if (!CVarRayTracingImagePlate.GetValueOnRenderThread())
+			{
+				return;
+			}
+
+			FMaterialRenderProxy* MaterialProxy = Material->GetRenderProxy();
+
+			if (RayTracingGeometry.RayTracingGeometryRHI.IsValid())
+			{
+				check(RayTracingGeometry.Initializer.IndexBuffer.IsValid());
+
+				FRayTracingInstance RayTracingInstance;
+				RayTracingInstance.Geometry = &RayTracingGeometry;
+				RayTracingInstance.InstanceTransforms.Add(GetLocalToWorld());
+
+				uint32 SectionIdx = 0;
+				FMeshBatch MeshBatch;
+
+				MeshBatch.VertexFactory = &VertexFactory;
+				MeshBatch.SegmentIndex = 0;
+				MeshBatch.MaterialRenderProxy = MaterialProxy;
+				MeshBatch.ReverseCulling = IsLocalToWorldDeterminantNegative();
+				MeshBatch.Type = PT_TriangleList;
+				MeshBatch.DepthPriorityGroup = SDPG_World;
+				MeshBatch.bCanApplyViewModeOverrides = false;
+				MeshBatch.CastRayTracedShadow = IsShadowCast(Context.ReferenceView);
+
+				FMeshBatchElement& BatchElement = MeshBatch.Elements[0];
+				BatchElement.IndexBuffer = &IndexBuffer;
+
+				bool bHasPrecomputedVolumetricLightmap = false;
+				FMatrix PreviousLocalToWorld = {};
+				PreviousLocalToWorld.SetIdentity();
+				int32 SingleCaptureIndex = 0;
+				bool bOutputVelocity = false;
+				GetScene().GetPrimitiveUniformShaderParameters_RenderThread(GetPrimitiveSceneInfo(), bHasPrecomputedVolumetricLightmap, PreviousLocalToWorld, SingleCaptureIndex, bOutputVelocity);
+				bOutputVelocity |= AlwaysHasVelocity();
+
+				FDynamicPrimitiveUniformBuffer& DynamicPrimitiveUniformBuffer = Context.RayTracingMeshResourceCollector.AllocateOneFrameResource<FDynamicPrimitiveUniformBuffer>();
+				DynamicPrimitiveUniformBuffer.Set(GetLocalToWorld(), PreviousLocalToWorld, GetBounds(), GetLocalBounds(), GetLocalBounds(), true, bHasPrecomputedVolumetricLightmap, bOutputVelocity, GetCustomPrimitiveData());
+				BatchElement.PrimitiveUniformBufferResource = &DynamicPrimitiveUniformBuffer.UniformBuffer;
+
+				BatchElement.FirstIndex = 0;
+				BatchElement.NumPrimitives = IndexBuffer.Indices.Num() / 3;
+				BatchElement.MinVertexIndex = 0;
+				BatchElement.MaxVertexIndex = VertexBuffers.PositionVertexBuffer.GetNumVertices() - 1;
+
+				RayTracingInstance.Materials.Add(MeshBatch);
+
+				RayTracingInstance.BuildInstanceMaskAndFlags(GetScene().GetFeatureLevel());
+				OutRayTracingInstances.Add(RayTracingInstance);
+			}
+		}
+#endif // RHI_RAYTRACING
+
 	private:
 		UMaterialInterface* Material;
 		FMaterialRelevance MaterialRelevance;
 		FStaticMeshVertexBuffers VertexBuffers;
 		FImagePlateIndexBuffer IndexBuffer;
 		FLocalVertexFactory VertexFactory;
+#if RHI_RAYTRACING
+		FRayTracingGeometry RayTracingGeometry;
+#endif
 	};
 }
 
