Image Plate: Add Settings dependency

diff --git a/Engine/Plugins/Experimental/ImagePlate/Source/ImagePlateEditor/ImagePlateEditor.Build.cs b/Engine/Plugins/Experimental/ImagePlate/Source/ImagePlateEditor/ImagePlateEditor.Build.cs
--- a/Engine/Plugins/Experimental/ImagePlate/Source/ImagePlateEditor/ImagePlateEditor.Build.cs
+++ b/Engine/Plugins/Experimental/ImagePlate/Source/ImagePlateEditor/ImagePlateEditor.Build.cs
@@ -16,6 +16,7 @@ namespace UnrealBuildTool.Rules
 					"Engine",
 					"ImagePlate",
 					"RHI",
+					"Settings",
 					"Slate",
 					"SlateCore",
                     "TimeManagement",
