Image Plate: Add back module includes

diff --git a/Engine/Plugins/Experimental/ImagePlate/Source/ImagePlateEditor/Private/ImagePlateEditor.cpp b/Engine/Plugins/Experimental/ImagePlate/Source/ImagePlateEditor/Private/ImagePlateEditor.cpp
--- a/Engine/Plugins/Experimental/ImagePlate/Source/ImagePlateEditor/Private/ImagePlateEditor.cpp
+++ b/Engine/Plugins/Experimental/ImagePlate/Source/ImagePlateEditor/Private/ImagePlateEditor.cpp
@@ -4,6 +4,8 @@
 #include "Brushes/SlateImageBrush.h"
 #include "Styling/SlateStyle.h"
 #include "Misc/Paths.h"
+#include "Modules/ModuleInterface.h"
+#include "Modules/ModuleManager.h"
 #include "Styling/SlateStyleRegistry.h"
 #include "ISettingsModule.h"
 
