[ImagePlate]
* Moved code one step closer to not using actors render side

diff --git a/Engine/Plugins/Experimental/ImagePlate/Source/ImagePlate/Private/ImagePlateFrustumComponent.cpp b/Engine/Plugins/Experimental/ImagePlate/Source/ImagePlate/Private/ImagePlateFrustumComponent.cpp
--- a/Engine/Plugins/Experimental/ImagePlate/Source/ImagePlate/Private/ImagePlateFrustumComponent.cpp
+++ b/Engine/Plugins/Experimental/ImagePlate/Source/ImagePlate/Private/ImagePlateFrustumComponent.cpp
@@ -26,21 +26,23 @@ public:
 
 	FImagePlateFrustumSceneProxy(const UImagePlateFrustumComponent* InComponent)
 		: FPrimitiveSceneProxy(InComponent)
-		, ViewTarget(nullptr)
 	{
 		bWillEverBeLit = false;
 
 		UImagePlateComponent* Parent = Cast<UImagePlateComponent>(InComponent->GetAttachParent());
 		if (Parent && Parent->GetPlate().bFillScreen)
 		{
-			ViewTarget = Parent->FindViewTarget();
+			if (AActor* ViewTarget = Parent->FindViewTarget())
+			{
+				ViewTargetID = ViewTarget->GetUniqueID();
+			}
 			InvViewProjectionMatrix = Parent->GetCachedInvViewProjectionMatrix();
 		}
 	}
 
 	virtual void GetDynamicMeshElements(const TArray<const FSceneView*>& Views, const FSceneViewFamily& ViewFamily, uint32 VisibilityMap, FMeshElementCollector& Collector) const override
 	{
-		if (!ViewTarget)
+		if (!ViewTargetID)
 		{
 			return;
 		}
@@ -55,7 +57,7 @@ public:
 			const FSceneView* View = Views[ViewIndex];
 
 			// Don't draw frustums on views that are set as the current view target
-			if (View->ViewActor == ViewTarget)
+			if (View->ViewActor.ActorUniqueId == ViewTargetID)
 			{
 				continue;
 			}
@@ -97,7 +99,7 @@ public:
 	uint32 GetAllocatedSize( void ) const { return( FPrimitiveSceneProxy::GetAllocatedSize() ); }
 
 private:
-	AActor* ViewTarget;
+	uint32 ViewTargetID = 0;
 	FMatrix InvViewProjectionMatrix;
 };
 
