Add shader resource flag to index buffers so they can be used for raytracing

diff --git a/Engine/Plugins/Experimental/ImagePlate/Source/ImagePlate/Private/ImagePlateComponent.cpp b/Engine/Plugins/Experimental/ImagePlate/Source/ImagePlate/Private/ImagePlateComponent.cpp
--- a/Engine/Plugins/Experimental/ImagePlate/Source/ImagePlate/Private/ImagePlateComponent.cpp
+++ b/Engine/Plugins/Experimental/ImagePlate/Source/ImagePlate/Private/ImagePlateComponent.cpp
@@ -40,7 +40,7 @@ namespace
 
 		virtual void InitRHI(FRHICommandListBase& RHICmdList) override
 		{
-			IndexBufferRHI = UE::RHIResourceUtils::CreateIndexBufferFromArray(RHICmdList, TEXT("FImagePlateIndexBuffer"), EBufferUsageFlags::Static, MakeConstArrayView(Indices));
+			IndexBufferRHI = UE::RHIResourceUtils::CreateIndexBufferFromArray(RHICmdList, TEXT("FImagePlateIndexBuffer"), EBufferUsageFlags::Static | EBufferUsageFlags::ShaderResource, MakeConstArrayView(Indices));
 		}
 	};
 
