UE: Load the default ImagePlate material in UImagePlateComponent's constructor instead of FImagePlateParameters' constructor. This load will only happen once and during CDO construction (i.e. a stable time)

diff --git a/Engine/Plugins/Experimental/ImagePlate/Source/ImagePlate/Private/ImagePlateComponent.cpp b/Engine/Plugins/Experimental/ImagePlate/Source/ImagePlate/Private/ImagePlateComponent.cpp
--- a/Engine/Plugins/Experimental/ImagePlate/Source/ImagePlate/Private/ImagePlateComponent.cpp
+++ b/Engine/Plugins/Experimental/ImagePlate/Source/ImagePlate/Private/ImagePlateComponent.cpp
@@ -2,6 +2,7 @@
 
 #include "ImagePlateComponent.h"
 
+#include "UObject/ConstructorHelpers.h"
 #include "Engine/Texture.h"
 #include "PrimitiveViewRelevance.h"
 #include "Materials/MaterialInstanceDynamic.h"
@@ -309,7 +310,7 @@ namespace
 }
 
 FImagePlateParameters::FImagePlateParameters()
-	: Material(LoadObject<UMaterialInterface>(nullptr, TEXT("/ImagePlate/DefaultImagePlateMaterial.DefaultImagePlateMaterial")))
+	: Material(nullptr)
 	, TextureParameterName("InputTexture")
 	, bFillScreen(true)
 	, FillScreenAmount(100, 100)
@@ -322,6 +323,9 @@ FImagePlateParameters::FImagePlateParameters()
 UImagePlateComponent::UImagePlateComponent(const FObjectInitializer& Init)
 	: Super(Init)
 {
+	static ConstructorHelpers::FObjectFinder<UMaterialInterface> DefaultImagePlateMaterial(TEXT("/ImagePlate/DefaultImagePlateMaterial"));
+	Plate.Material = DefaultImagePlateMaterial.Object;
+
 	SetCollisionProfileName(UCollisionProfile::NoCollision_ProfileName);
 	bUseAsOccluder = false;
 	bTickInEditor = true;
@@ -372,6 +376,12 @@ void UImagePlateComponent::SetImagePlate(FImagePlateParameters NewPlate)
 {
 	Plate = NewPlate;
 
+	// Interpret a null material to mean the default material
+	if (!Plate.Material)
+	{
+		Plate.Material = GetDefault<UImagePlateComponent>()->Plate.Material;
+	}
+
 	UpdateMaterialParametersForMedia();
 }
 
