[HWRT] Fix crash when UImagePlateComponent doesn't have a valid material assigned.

#Synced-CL 52071961

diff --git a/Engine/Plugins/Experimental/ImagePlate/Source/ImagePlate/Private/ImagePlateComponent.cpp b/Engine/Plugins/Experimental/ImagePlate/Source/ImagePlate/Private/ImagePlateComponent.cpp
--- a/Engine/Plugins/Experimental/ImagePlate/Source/ImagePlate/Private/ImagePlateComponent.cpp
+++ b/Engine/Plugins/Experimental/ImagePlate/Source/ImagePlate/Private/ImagePlateComponent.cpp
@@ -65,6 +65,10 @@ namespace
 			{
 				MaterialRelevance |= Material->GetRelevance_Concurrent(GetScene().GetShaderPlatform());
 			}
+
+#if RHI_RAYTRACING
+			bHasRayTracingRepresentation = (Material != nullptr) && IsRayTracingEnabled();
+#endif // RHI_RAYTRACING
 		}
 
 		~FImagePlateSceneProxy()
@@ -76,15 +80,16 @@ namespace
 			VertexFactory.ReleaseResource();
 #if RHI_RAYTRACING
 			RayTracingGeometry.ReleaseResource();
-#endif
+#endif // RHI_RAYTRACING
 		}
 
 		virtual void CreateRenderThreadResources(FRHICommandListBase& RHICmdList) override
 		{
 			BuildMesh(RHICmdList);
 			IndexBuffer.InitResource(RHICmdList);
+
 #if RHI_RAYTRACING
-			if (IsRayTracingEnabled())
+			if (bHasRayTracingRepresentation)
 			{
 				FRayTracingGeometryInitializer Initializer;
 				Initializer.DebugName = GetOwnerName();
@@ -104,7 +109,7 @@ namespace
 				RayTracingGeometry.SetInitializer(Initializer);
 				RayTracingGeometry.InitResource(RHICmdList);
 			}
-#endif
+#endif // RHI_RAYTRACING
 		}
 
 		void BuildMesh(FRHICommandListBase& RHICmdList)
@@ -244,7 +249,7 @@ namespace
 #if RHI_RAYTRACING
 		virtual bool IsRayTracingRelevant() const override { return true; }
 
-		virtual bool HasRayTracingRepresentation() const override { return true; }
+		virtual bool HasRayTracingRepresentation() const override { return bHasRayTracingRepresentation; }
 
 		virtual void GetDynamicRayTracingInstances(FRayTracingInstanceCollector& Collector) override final
 		{
@@ -323,7 +328,8 @@ namespace
 		FLocalVertexFactory VertexFactory;
 #if RHI_RAYTRACING
 		FRayTracingGeometry RayTracingGeometry;
-#endif
+		bool bHasRayTracingRepresentation;
+#endif // RHI_RAYTRACING
 	};
 }
 
